Integrating CompositePlane plugins to 4.25

diff --git a/Engine/Plugins/VirtualProduction/CompositePlane/CompositePlane.uplugin b/Engine/Plugins/VirtualProduction/CompositePlane/CompositePlane.uplugin
new file mode 100644
--- /dev/null
+++ b/Engine/Plugins/VirtualProduction/CompositePlane/CompositePlane.uplugin
@@ -0,0 +1,24 @@
+{
+	"FileVersion": 3,
+	"Version": 1,
+	"VersionName": "1.0",
+	"FriendlyName": "Composite Plane",
+	"Description": "Provides a cine camera actor for projecting textures and videos",
+	"Category": "Virtual Production",
+	"CreatedBy": "Epic Games, Inc.",
+	"CreatedByURL": "http://epicgames.com",
+	"DocsURL": "",
+	"MarketplaceURL": "",
+	"SupportURL": "",
+	"CanContainContent": true,
+	"IsBetaVersion": true,
+	"IsExperimentalVersion": false,
+	"Installed": false,
+	"Modules": [
+		{
+			"Name": "CompositePlane",
+			"Type": "Editor",
+			"LoadingPhase": "PostEngineInit"
+		}
+	]
+}
\ No newline at end of file
diff --git a/Engine/Plugins/VirtualProduction/CompositePlane/Source/CompositePlane/CompositePlane.Build.cs b/Engine/Plugins/VirtualProduction/CompositePlane/Source/CompositePlane/CompositePlane.Build.cs
new file mode 100644
--- /dev/null
+++ b/Engine/Plugins/VirtualProduction/CompositePlane/Source/CompositePlane/CompositePlane.Build.cs
@@ -0,0 +1,31 @@
+// Copyright Epic Games, Inc. All Rights Reserved.
+
+using UnrealBuildTool;
+
+public class CompositePlane : ModuleRules
+{
+	public CompositePlane(ReadOnlyTargetRules Target) : base(Target)
+	{
+		PrivateDependencyModuleNames.AddRange(
+			new string[]
+			{
+                "Core",
+                "CoreUObject",
+				"Engine",
+                "Projects",
+                "Slate",
+				"SlateCore",
+			}
+			);
+
+        if (Target.bBuildEditor)
+        {
+            PrivateDependencyModuleNames.AddRange(
+                new string[] {
+                        "UnrealEd",
+                        "PlacementMode",
+                }
+            );
+        }
+	}
+}
\ No newline at end of file
diff --git a/Engine/Plugins/VirtualProduction/CompositePlane/Source/CompositePlane/Private/CompositePlane.cpp b/Engine/Plugins/VirtualProduction/CompositePlane/Source/CompositePlane/Private/CompositePlane.cpp
new file mode 100644
--- /dev/null
+++ b/Engine/Plugins/VirtualProduction/CompositePlane/Source/CompositePlane/Private/CompositePlane.cpp
@@ -0,0 +1,20 @@
+// Copyright Epic Games, Inc. All Rights Reserved.
+
+#include "CompositePlane.h"
+#include "CompositePlanePlacement.h"
+
+#define LOCTEXT_NAMESPACE "FCompositePlaneModule"
+
+void FCompositePlaneModule::StartupModule()
+{
+	FCompositePlanePlacement::RegisterPlacement();
+}
+
+void FCompositePlaneModule::ShutdownModule()
+{
+	FCompositePlanePlacement::UnregisterPlacement();
+}
+
+#undef LOCTEXT_NAMESPACE
+	
+IMPLEMENT_MODULE(FCompositePlaneModule, CompositePlane)
\ No newline at end of file
diff --git a/Engine/Plugins/VirtualProduction/CompositePlane/Source/CompositePlane/Private/CompositePlane.h b/Engine/Plugins/VirtualProduction/CompositePlane/Source/CompositePlane/Private/CompositePlane.h
new file mode 100644
--- /dev/null
+++ b/Engine/Plugins/VirtualProduction/CompositePlane/Source/CompositePlane/Private/CompositePlane.h
@@ -0,0 +1,15 @@
+// Copyright Epic Games, Inc. All Rights Reserved.
+
+#pragma once
+
+#include "CoreMinimal.h"
+#include "Modules/ModuleManager.h"
+
+class FCompositePlaneModule : public IModuleInterface
+{
+public:
+
+	/** IModuleInterface implementation */
+	virtual void StartupModule() override;
+	virtual void ShutdownModule() override;
+};
diff --git a/Engine/Plugins/VirtualProduction/CompositePlane/Source/CompositePlane/Private/CompositePlanePlacement.cpp b/Engine/Plugins/VirtualProduction/CompositePlane/Source/CompositePlane/Private/CompositePlanePlacement.cpp
new file mode 100644
--- /dev/null
+++ b/Engine/Plugins/VirtualProduction/CompositePlane/Source/CompositePlane/Private/CompositePlanePlacement.cpp
@@ -0,0 +1,44 @@
+// Copyright Epic Games, Inc. All Rights Reserved.
+
+#include "CompositePlanePlacement.h"
+#include "IPlacementModeModule.h"
+#include "ActorFactories/ActorFactoryBlueprint.h"
+
+#define LOCTEXT_NAMESPACE "FCompositePlaneModule"
+
+void FCompositePlanePlacement::RegisterPlacement()
+{
+	UBlueprint* CompositePlane = Cast<UBlueprint>(FSoftObjectPath(TEXT("/CompositePlane/BP_CineCameraProj.BP_CineCameraProj")).TryLoad());
+	if (CompositePlane == nullptr)
+	{
+		return;
+	}
+
+	IPlacementModeModule& PlacementModeModule = IPlacementModeModule::Get();
+	const FPlacementCategoryInfo* Info = PlacementModeModule.GetRegisteredPlacementCategory(FName("Cinematic"));
+	if (Info == nullptr)
+	{
+		return;
+	}
+
+	FPlaceableItem* BPPlacement = new FPlaceableItem(
+		*UActorFactoryBlueprint::StaticClass(),
+		FAssetData(CompositePlane, true),
+		FName(""),
+		TOptional<FLinearColor>(),
+		TOptional<int32>(),
+		NSLOCTEXT("PlacementMode", "Composite Plane", "Composite Plane")
+	);
+
+	IPlacementModeModule::Get().RegisterPlaceableItem(Info->UniqueHandle, MakeShareable(BPPlacement));
+}
+
+void FCompositePlanePlacement::UnregisterPlacement()
+{
+	if (IPlacementModeModule::IsAvailable())
+	{
+		IPlacementModeModule::Get().UnregisterPlacementCategory(FName("Cinematic"));
+	}
+}
+
+#undef LOCTEXT_NAMESPACE
\ No newline at end of file
diff --git a/Engine/Plugins/VirtualProduction/CompositePlane/Source/CompositePlane/Private/CompositePlanePlacement.h b/Engine/Plugins/VirtualProduction/CompositePlane/Source/CompositePlane/Private/CompositePlanePlacement.h
new file mode 100644
--- /dev/null
+++ b/Engine/Plugins/VirtualProduction/CompositePlane/Source/CompositePlane/Private/CompositePlanePlacement.h
@@ -0,0 +1,11 @@
+// Copyright Epic Games, Inc. All Rights Reserved.
+#pragma once
+
+#include "CoreMinimal.h"
+
+class FCompositePlanePlacement
+{
+public:
+	static void RegisterPlacement();
+	static void UnregisterPlacement();
+};
\ No newline at end of file
