CompositePlane: Start deprecation process.

* Added deprecation tag which only emits a deprecation display message to the user.

diff --git a/Engine/Plugins/VirtualProduction/CompositePlane/CompositePlane.uplugin b/Engine/Plugins/VirtualProduction/CompositePlane/CompositePlane.uplugin
--- a/Engine/Plugins/VirtualProduction/CompositePlane/CompositePlane.uplugin
+++ b/Engine/Plugins/VirtualProduction/CompositePlane/CompositePlane.uplugin
@@ -3,7 +3,7 @@
 	"Version": 1,
 	"VersionName": "1.0",
 	"FriendlyName": "Composite Plane",
-	"Description": "Provides a cine camera actor for projecting textures and videos",
+	"Description": "This plugin has been deprecated and will be removed in a future engine version. Please update your project to use the features of the new Composure plugin instead (Composite.uplugin).\nProvides a cine camera actor for projecting textures and videos",
 	"Category": "Virtual Production",
 	"CreatedBy": "Epic Games, Inc.",
 	"CreatedByURL": "https://epicgames.com",
@@ -14,6 +14,7 @@
 	"IsBetaVersion": true,
 	"IsExperimentalVersion": false,
 	"Installed": false,
+	"DeprecatedEngineVersion": "5.7",
 	"Modules": [
 		{
 			"Name": "CompositePlane",
